Disable function that is in RAD tools code but not UE for SymsLIB

diff --git a/Engine/Source/Programs/SymsLibDump/Private/SymsLibDump.cpp b/Engine/Source/Programs/SymsLibDump/Private/SymsLibDump.cpp
--- a/Engine/Source/Programs/SymsLibDump/Private/SymsLibDump.cpp
+++ b/Engine/Source/Programs/SymsLibDump/Private/SymsLibDump.cpp
@@ -274,7 +274,8 @@ namespace
 						Symbols[ProcIndex].Size    = RangeArray.ranges[0].max - RangeArray.ranges[0].min;
 					}
 
-					SYMS_String8 Name = syms_linkage_name_from_sid(Arena, Group->dbg_data, Group->dbg, UnitAccel, SymbolID);
+					// TODO need new changes from RAD tools
+					SYMS_String8 Name {nullptr, 0}; // = syms_linkage_name_from_sid(Arena, Group->dbg_data, Group->dbg, UnitAccel, SymbolID);
 
 					// If we fail to find a linkage name fallback to name from sid. Some platforms, like Windows wont have a linkage name
 					if  (Name.size == 0)
